UMG - Adding the Web Browser Widget plugin back, renamed and restructured for UMG specifically.

diff --git a/Engine/Plugins/Runtime/WebBrowserWidget/Source/WebBrowserWidget/Private/WebBrowser.cpp b/Engine/Plugins/Runtime/WebBrowserWidget/Source/WebBrowserWidget/Private/WebBrowser.cpp
new file mode 100644
--- /dev/null
+++ b/Engine/Plugins/Runtime/WebBrowserWidget/Source/WebBrowserWidget/Private/WebBrowser.cpp
@@ -0,0 +1,97 @@
+// Copyright 1998-2015 Epic Games, Inc. All Rights Reserved.
+
+#include "WebBrowserWidgetPrivatePCH.h"
+
+#include "SWebBrowser.h"
+
+#include "WebBrowser.h"
+
+#define LOCTEXT_NAMESPACE "WebBrowser"
+
+/////////////////////////////////////////////////////
+// UWebBrowser
+
+UWebBrowser::UWebBrowser(const FObjectInitializer& ObjectInitializer)
+	: Super(ObjectInitializer)
+{
+	bIsVariable = true;
+}
+
+void UWebBrowser::LoadURL(FString NewURL)
+{
+	if ( WebBrowserWidget.IsValid() )
+	{
+		return WebBrowserWidget->LoadURL(NewURL);
+	}
+}
+
+void UWebBrowser::LoadString(FString Contents, FString DummyURL)
+{
+	if ( WebBrowserWidget.IsValid() )
+	{
+		return WebBrowserWidget->LoadString(Contents, DummyURL);
+	}
+}
+
+FText UWebBrowser::GetTitleText() const
+{
+	if ( WebBrowserWidget.IsValid() )
+	{
+		return WebBrowserWidget->GetTitleText();
+	}
+
+	return FText::GetEmpty();
+}
+
+void UWebBrowser::ReleaseSlateResources(bool bReleaseChildren)
+{
+	Super::ReleaseSlateResources(bReleaseChildren);
+
+	WebBrowserWidget.Reset();
+}
+
+TSharedRef<SWidget> UWebBrowser::RebuildWidget()
+{
+	if ( IsDesignTime() )
+	{
+		return BuildDesignTimeWidget(SNew(SBox)
+			.HAlign(HAlign_Center)
+			.VAlign(VAlign_Center)
+			[
+				SNew(STextBlock)
+				.Text(LOCTEXT("Web Browser", "Web Browser"))
+			]);
+	}
+	else
+	{
+		WebBrowserWidget = SNew(SWebBrowser)
+			.InitialURL(InitialURL)
+			.ShowControls(false)
+			.SupportsTransparency(bSupportsTransparency);
+
+		return WebBrowserWidget.ToSharedRef();
+	}
+}
+
+void UWebBrowser::SynchronizeProperties()
+{
+	Super::SynchronizeProperties();
+
+	if ( WebBrowserWidget.IsValid() )
+	{
+
+	}
+}
+
+#if WITH_EDITOR
+
+const FText UWebBrowser::GetPaletteCategory()
+{
+	return LOCTEXT("Experimental", "Experimental");
+}
+
+#endif
+
+/////////////////////////////////////////////////////
+
+#undef LOCTEXT_NAMESPACE
diff --git a/Engine/Plugins/Runtime/WebBrowserWidget/Source/WebBrowserWidget/Private/WebBrowserWidgetModule.cpp b/Engine/Plugins/Runtime/WebBrowserWidget/Source/WebBrowserWidget/Private/WebBrowserWidgetModule.cpp
new file mode 100644
--- /dev/null
+++ b/Engine/Plugins/Runtime/WebBrowserWidget/Source/WebBrowserWidget/Private/WebBrowserWidgetModule.cpp
@@ -0,0 +1,23 @@
+// Copyright 1998-2015 Epic Games, Inc. All Rights Reserved.
+
+#include "WebBrowserWidgetPrivatePCH.h"
+
+//////////////////////////////////////////////////////////////////////////
+// FWebBrowserWidgetModule
+
+class FWebBrowserWidgetModule : public IWebBrowserWidgetModule
+{
+public:
+	virtual void StartupModule() override
+	{
+	}
+
+	virtual void ShutdownModule() override
+	{
+	}
+};
+
+//////////////////////////////////////////////////////////////////////////
+
+IMPLEMENT_MODULE(FWebBrowserWidgetModule, WebBrowserWidget);
+DEFINE_LOG_CATEGORY(LogWebBrowserWidget);
diff --git a/Engine/Plugins/Runtime/WebBrowserWidget/Source/WebBrowserWidget/Private/WebBrowserWidgetPrivatePCH.h b/Engine/Plugins/Runtime/WebBrowserWidget/Source/WebBrowserWidget/Private/WebBrowserWidgetPrivatePCH.h
new file mode 100644
--- /dev/null
+++ b/Engine/Plugins/Runtime/WebBrowserWidget/Source/WebBrowserWidget/Private/WebBrowserWidgetPrivatePCH.h
@@ -0,0 +1,10 @@
+// Copyright 1998-2015 Epic Games, Inc. All Rights Reserved.
+
+#pragma once
+
+#include "SlateBasics.h"
+#include "SlateCore.h"
+
+#include "WebBrowserWidgetModule.h"
+
+DECLARE_LOG_CATEGORY_EXTERN(LogWebBrowserWidget, Verbose, All);
\ No newline at end of file
diff --git a/Engine/Plugins/Runtime/WebBrowserWidget/Source/WebBrowserWidget/Public/WebBrowser.h b/Engine/Plugins/Runtime/WebBrowserWidget/Source/WebBrowserWidget/Public/WebBrowser.h
new file mode 100644
--- /dev/null
+++ b/Engine/Plugins/Runtime/WebBrowserWidget/Source/WebBrowserWidget/Public/WebBrowser.h
@@ -0,0 +1,68 @@
+// Copyright 1998-2015 Epic Games, Inc. All Rights Reserved.
+
+#pragma once
+
+#include "Widget.h"
+
+#include "WebBrowser.generated.h"
+
+/**
+ * 
+ */
+UCLASS()
+class WEBBROWSERWIDGET_API UWebBrowser : public UWidget
+{
+	GENERATED_UCLASS_BODY()
+
+public:
+
+	/**
+	 * Load the specified URL
+	 *
+	 * @param NewURL New URL to load
+	 */
+	UFUNCTION(BlueprintCallable, Category="Web Browser")
+	void LoadURL(FString NewURL);
+
+	/**
+	 * Load a string as data to create a web page
+	 *
+	 * @param Contents String to load
+	 * @param DummyURL Dummy URL for the page
+	 */
+	UFUNCTION(BlueprintCallable, Category="Web Browser")
+	void LoadString(FString Contents, FString DummyURL);
+
+	/**
+	 * Get the current title of the web page
+	 */
+	UFUNCTION(BlueprintCallable, Category="Web Browser")
+	FText GetTitleText() const;
+
+public:
+
+	// UWidget interface
+	virtual void SynchronizeProperties() override;
+	// End of UWidget interface
+
+	virtual void ReleaseSlateResources(bool bReleaseChildren) override;
+
+#if WITH_EDITOR
+	virtual const FText GetPaletteCategory() override;
+#endif
+
+protected:
+	UPROPERTY(EditAnywhere, Category=Appearance)
+	FString InitialURL;
+
+	UPROPERTY(EditAnywhere, Category=Appearance)
+	bool bSupportsTransparency;
+
+protected:
+	TSharedPtr<class SWebBrowser> WebBrowserWidget;
+
+protected:
+	// UWidget interface
+	virtual TSharedRef<SWidget> RebuildWidget() override;
+	// End of UWidget interface
+};
diff --git a/Engine/Plugins/Runtime/WebBrowserWidget/Source/WebBrowserWidget/Public/WebBrowserWidgetModule.h b/Engine/Plugins/Runtime/WebBrowserWidget/Source/WebBrowserWidget/Public/WebBrowserWidgetModule.h
new file mode 100644
--- /dev/null
+++ b/Engine/Plugins/Runtime/WebBrowserWidget/Source/WebBrowserWidget/Public/WebBrowserWidgetModule.h
@@ -0,0 +1,37 @@
+// Copyright 1998-2015 Epic Games, Inc. All Rights Reserved.
+
+#pragma once
+
+#include "ModuleManager.h"
+
+
+/**
+ * The public interface to this module
+ */
+class IWebBrowserWidgetModule : public IModuleInterface
+{
+
+public:
+
+	/**
+	 * Singleton-like access to this module's interface.  This is just for convenience!
+	 * Beware of calling this during the shutdown phase, though.  Your module might have been unloaded already.
+	 *
+	 * @return Returns singleton instance, loading the module on demand if needed
+	 */
+	static inline IWebBrowserWidgetModule& Get()
+	{
+		return FModuleManager::LoadModuleChecked<IWebBrowserWidgetModule>( "WebBrowserWidget" );
+	}
+
+	/**
+	 * Checks to see if this module is loaded and ready.  It is only valid to call Get() if IsAvailable() returns true.
+	 *
+	 * @return True if the module is loaded and ready to use
+	 */
+	static inline bool IsAvailable()
+	{
+		return FModuleManager::Get().IsModuleLoaded( "WebBrowserWidget" );
+	}
+};
+
diff --git a/Engine/Plugins/Runtime/WebBrowserWidget/Source/WebBrowserWidget/WebBrowserWidget.build.cs b/Engine/Plugins/Runtime/WebBrowserWidget/Source/WebBrowserWidget/WebBrowserWidget.build.cs
new file mode 100644
--- /dev/null
+++ b/Engine/Plugins/Runtime/WebBrowserWidget/Source/WebBrowserWidget/WebBrowserWidget.build.cs
@@ -0,0 +1,22 @@
+// Copyright 1998-2015 Epic Games, Inc. All Rights Reserved.
+
+namespace UnrealBuildTool.Rules
+{
+	public class WebBrowserWidget : ModuleRules
+	{
+		public WebBrowserWidget(TargetInfo Target)
+		{
+			PublicDependencyModuleNames.AddRange(
+				new string[]
+				{
+					"Core",
+				    "CoreUObject",
+					"WebBrowser",
+                    "Slate",
+                    "SlateCore",
+					"UMG"
+				}
+			);
+		}
+	}
+}
diff --git a/Engine/Plugins/Runtime/WebBrowserWidget/WebBrowserWidget.uplugin b/Engine/Plugins/Runtime/WebBrowserWidget/WebBrowserWidget.uplugin
new file mode 100644
--- /dev/null
+++ b/Engine/Plugins/Runtime/WebBrowserWidget/WebBrowserWidget.uplugin
@@ -0,0 +1,28 @@
+{
+	"FileVersion": 3,
+	"Version": 1,
+	"VersionName": "1.0",
+	"FriendlyName": "Web Browser",
+	"Description": "Allows the user to create a Web Browser Widget",
+	"Category": "Widgets",
+	"CreatedBy": "Epic Games, Inc.",
+	"CreatedByURL": "http://epicgames.com",
+	"DocsURL": "",
+	"MarketplaceURL": "",
+	"Modules": [
+		{
+			"Name": "WebBrowserWidget",
+			"Type": "Runtime",
+			"LoadingPhase": "Default",
+			"WhitelistPlatforms": [
+				"Win64",
+				"Win32",
+				"Mac"
+			]
+		}
+	],
+	"EnabledByDefault": false,
+	"CanContainContent": false,
+	"IsBetaVersion": false,
+	"Installed": false
+}
\ No newline at end of file
