Fixes for WebTexture_T cooking issues
#ue4

diff --git a/Engine/Config/BaseEngine.ini b/Engine/Config/BaseEngine.ini
--- a/Engine/Config/BaseEngine.ini
+++ b/Engine/Config/BaseEngine.ini
@@ -481,6 +481,7 @@ bUseStreamingPause=false
 +FunctionRedirects=(OldName="InverseLerp",NewName="/Script/Engine.KismetMathLibrary.NormalizeToRange")
 +PropertyRedirects=(OldName="NormalizeToRange.A",NewName="RangeMin")
 +PropertyRedirects=(OldName="NormalizeToRange.B",NewName="RangeMax")
++ClassRedirects=(OldName="WebBrowserTexture",NewName="/Script/WebBrowserTexture.WebBrowserTexture")
 
 [CoreUObject.Metadata]
 ; Note: UnrealHeaderTool should be rerun after making changes to MetadataRedirects to catch any keys specified in class headers
diff --git a/Engine/Plugins/Runtime/WebBrowserWidget/Source/WebBrowserWidget/Private/WebBrowser.cpp b/Engine/Plugins/Runtime/WebBrowserWidget/Source/WebBrowserWidget/Private/WebBrowser.cpp
--- a/Engine/Plugins/Runtime/WebBrowserWidget/Source/WebBrowserWidget/Private/WebBrowser.cpp
+++ b/Engine/Plugins/Runtime/WebBrowserWidget/Source/WebBrowserWidget/Private/WebBrowser.cpp
@@ -19,6 +19,10 @@
 #include "PackageHelperFunctions.h"
 #endif
 
+#if WITH_EDITOR || PLATFORM_ANDROID
+#include "WebBrowserTexture.h"
+#endif
+
 #define LOCTEXT_NAMESPACE "WebBrowser"
 
 /////////////////////////////////////////////////////
@@ -28,6 +32,8 @@ UWebBrowser::UWebBrowser(const FObjectInitializer& ObjectInitializer)
 	: Super(ObjectInitializer)
 {
 	bIsVariable = true;
+
+#if WITH_EDITOR || PLATFORM_ANDROID
 	struct FConstructorStatics
 	{
 		ConstructorHelpers::FObjectFinder<UObject>			DefaultTextureMaterial;
@@ -37,7 +43,11 @@ UWebBrowser::UWebBrowser(const FObjectInitializer& ObjectInitializer)
 	};
 	static FConstructorStatics ConstructorStatics;
 
+	// Add a hard reference to UWebBrowserTexture, without this the WebBrowserTexture DLL never gets loaded on Windows.
+	UWebBrowserTexture::StaticClass();
+
 	DefaultMaterial = (UMaterial*)ConstructorStatics.DefaultTextureMaterial.Object;
+#endif
 }
 
 void UWebBrowser::LoadURL(FString NewURL)
diff --git a/Engine/Plugins/Runtime/WebBrowserWidget/Source/WebBrowserWidget/WebBrowserWidget.build.cs b/Engine/Plugins/Runtime/WebBrowserWidget/Source/WebBrowserWidget/WebBrowserWidget.build.cs
--- a/Engine/Plugins/Runtime/WebBrowserWidget/Source/WebBrowserWidget/WebBrowserWidget.build.cs
+++ b/Engine/Plugins/Runtime/WebBrowserWidget/Source/WebBrowserWidget/WebBrowserWidget.build.cs
@@ -21,7 +21,7 @@ namespace UnrealBuildTool.Rules
 
             if (Target.bBuildEditor == true)
             {
-                //@TODO: Needed for the triangulation code used for sprites (but only in editor mode)
+                //@TODO: UnrealEd Needed for the triangulation code used for sprites (but only in editor mode)
                 //@TOOD: Try to move the code dependent on the triangulation code to the editor-only module
                 PrivateIncludePathModuleNames.AddRange(
                     new string[]
@@ -36,6 +36,23 @@ namespace UnrealBuildTool.Rules
                     }
                 );
             }
+			
+			if (Target.bBuildEditor || Target.Platform == UnrealTargetPlatform.Android)
+			{
+                // WebBrowserTexture required for cooking Android
+                PrivateIncludePathModuleNames.AddRange(
+                    new string[]
+                    {
+                        "WebBrowserTexture",
+                    }
+                );
+                PrivateDependencyModuleNames.AddRange(
+                    new string[]
+                    {
+                        "WebBrowserTexture",
+                    }
+                );
+			}
         }
     }
 }
\ No newline at end of file
diff --git a/Engine/Source/Runtime/WebBrowserTexture/WebBrowserTexture.Build.cs b/Engine/Source/Runtime/WebBrowserTexture/WebBrowserTexture.Build.cs
--- a/Engine/Source/Runtime/WebBrowserTexture/WebBrowserTexture.Build.cs
+++ b/Engine/Source/Runtime/WebBrowserTexture/WebBrowserTexture.Build.cs
@@ -7,7 +7,9 @@ public class WebBrowserTexture : ModuleRules
 {
 	public WebBrowserTexture(ReadOnlyTargetRules Target) : base(Target)
 	{
-        if (Target.Platform == UnrealTargetPlatform.Android)
+        // WebBrowserTexture objects are needed only on Android, but we also need to be able to
+        // cook the asset so we must include it in editor builds
+        if (Target.Platform == UnrealTargetPlatform.Android || Target.bBuildEditor == true)
 		{
 			PublicIncludePaths.Add("Runtime/WebBrowserTexture/Public");
 			PrivateIncludePaths.Add("Runtime/WebBrowserTexture/Private");
