[Build] UE4 Dev-Mobile - CL 4339590 - Incremental Editor Win64

diff --git a/Engine/Plugins/Runtime/WebBrowserWidget/Source/WebBrowserWidget/Private/WebBrowserAssetManager.cpp b/Engine/Plugins/Runtime/WebBrowserWidget/Source/WebBrowserWidget/Private/WebBrowserAssetManager.cpp
new file mode 100644
--- /dev/null
+++ b/Engine/Plugins/Runtime/WebBrowserWidget/Source/WebBrowserWidget/Private/WebBrowserAssetManager.cpp
@@ -0,0 +1,39 @@
+// Copyright 1998-2018 Epic Games, Inc. All Rights Reserved.
+
+#include "WebBrowserAssetManager.h"
+
+#if WITH_EDITOR || PLATFORM_ANDROID || PLATFORM_IOS
+#include "Materials/Material.h"
+#include "WebBrowserTexture.h"
+#endif
+
+/////////////////////////////////////////////////////
+// WebBrowserAssetManager
+
+UWebBrowserAssetManager::UWebBrowserAssetManager(const FObjectInitializer& ObjectInitializer)
+	: Super(ObjectInitializer) ,
+	DefaultMaterial(FString(TEXT("/WebBrowserWidget/WebTexture_M.WebTexture_M"))) ,
+	DefaultTranslucentMaterial(FString(TEXT("/WebBrowserWidget/WebTexture_TM.WebTexture_M")))
+{
+#if WITH_EDITOR || PLATFORM_ANDROID || PLATFORM_IOS
+	// Add a hard reference to UWebBrowserTexture, without this the WebBrowserTexture DLL never gets loaded on Windows.
+	UWebBrowserTexture::StaticClass();
+
+#endif
+};
+
+void UWebBrowserAssetManager::LoadDefaultMaterials()
+{
+	DefaultMaterial.LoadSynchronous();
+	DefaultTranslucentMaterial.LoadSynchronous();
+}
+
+UMaterial* UWebBrowserAssetManager::GetDefaultMaterial()
+{
+	return DefaultMaterial.Get();
+}
+
+UMaterial* UWebBrowserAssetManager::GetDefaultTranslucentMaterial()
+{
+	return DefaultTranslucentMaterial.Get();
+}
diff --git a/Engine/Plugins/Runtime/WebBrowserWidget/Source/WebBrowserWidget/Public/WebBrowserAssetManager.h b/Engine/Plugins/Runtime/WebBrowserWidget/Source/WebBrowserWidget/Public/WebBrowserAssetManager.h
new file mode 100644
--- /dev/null
+++ b/Engine/Plugins/Runtime/WebBrowserWidget/Source/WebBrowserWidget/Public/WebBrowserAssetManager.h
@@ -0,0 +1,28 @@
+﻿// Copyright 1998-2018 Epic Games, Inc. All Rights Reserved.
+
+#pragma once
+
+#include "UObject/SoftObjectPtr.h"
+
+#include "WebBrowserAssetManager.generated.h"
+
+class UMaterial;
+/**
+ * 
+ */
+UCLASS()
+class WEBBROWSERWIDGET_API UWebBrowserAssetManager : public UObject
+{
+	GENERATED_UCLASS_BODY()
+
+public:
+	void LoadDefaultMaterials();
+
+	UMaterial* GetDefaultMaterial(); 
+	UMaterial* GetDefaultTranslucentMaterial(); 
+
+protected:
+	UPROPERTY()
+	TSoftObjectPtr<UMaterial> DefaultMaterial;
+	TSoftObjectPtr<UMaterial> DefaultTranslucentMaterial;
+};
