UE-69711 WebTexture asset errors when packaging for Windows (fallout from UE-61554)

#Andoird
#iOS
#UE4
#4.22

The SlateUI.MaterialExpressionTextureBase for the WebTexture materials is now linked to the WebTexture
Renamed the transparent material as WebTexture_TM to correctly identify the material in the Content Browser
Put the UWebBrowserAssetManager unde WITH_EDITOR || PLATFORM_ANDROID || PLATFORM_IOS

From CL4958083/Dev-Mobile

diff --git a/Engine/Plugins/Runtime/WebBrowserWidget/Source/WebBrowserWidget/Private/WebBrowserAssetManager.cpp b/Engine/Plugins/Runtime/WebBrowserWidget/Source/WebBrowserWidget/Private/WebBrowserAssetManager.cpp
--- a/Engine/Plugins/Runtime/WebBrowserWidget/Source/WebBrowserWidget/Private/WebBrowserAssetManager.cpp
+++ b/Engine/Plugins/Runtime/WebBrowserWidget/Source/WebBrowserWidget/Private/WebBrowserAssetManager.cpp
@@ -1,24 +1,19 @@
 // Copyright 1998-2019 Epic Games, Inc. All Rights Reserved.
 
-#include "WebBrowserAssetManager.h"
-
 #if WITH_EDITOR || PLATFORM_ANDROID || PLATFORM_IOS
+#include "WebBrowserAssetManager.h"
 #include "WebBrowserTexture.h"
-#endif
 
 /////////////////////////////////////////////////////
 // WebBrowserAssetManager
 
 UWebBrowserAssetManager::UWebBrowserAssetManager(const FObjectInitializer& ObjectInitializer)
 	: Super(ObjectInitializer) ,
-	DefaultMaterial(FString(TEXT("/WebBrowserWidget/WebTexture_M.WebTexture_M"))) ,
-	DefaultTranslucentMaterial(FString(TEXT("/WebBrowserWidget/WebTexture_TM.WebTexture_M")))
+	DefaultMaterial(FString(TEXT("/WebBrowserWidget/WebTexture_M.WebTexture_M"))),
+	DefaultTranslucentMaterial(FString(TEXT("/WebBrowserWidget/WebTexture_TM.WebTexture_TM")))
 {
-#if WITH_EDITOR || PLATFORM_ANDROID || PLATFORM_IOS
 	// Add a hard reference to UWebBrowserTexture, without this the WebBrowserTexture DLL never gets loaded on Windows.
 	UWebBrowserTexture::StaticClass();
-
-#endif
 };
 
 void UWebBrowserAssetManager::LoadDefaultMaterials()
@@ -36,3 +31,4 @@ UMaterial* UWebBrowserAssetManager::GetDefaultTranslucentMaterial()
 {
 	return DefaultTranslucentMaterial.Get();
 }
+#endif
diff --git a/Engine/Plugins/Runtime/WebBrowserWidget/Source/WebBrowserWidget/Private/WebBrowserWidgetModule.cpp b/Engine/Plugins/Runtime/WebBrowserWidget/Source/WebBrowserWidget/Private/WebBrowserWidgetModule.cpp
--- a/Engine/Plugins/Runtime/WebBrowserWidget/Source/WebBrowserWidget/Private/WebBrowserWidgetModule.cpp
+++ b/Engine/Plugins/Runtime/WebBrowserWidget/Source/WebBrowserWidget/Private/WebBrowserWidgetModule.cpp
@@ -15,6 +15,7 @@ class FWebBrowserWidgetModule : public IWebBrowserWidgetModule
 public:
 	virtual void StartupModule() override
 	{
+#if WITH_EDITOR || PLATFORM_ANDROID || PLATFORM_IOS
 		if (WebBrowserAssetMgr == nullptr)
 		{
 			WebBrowserAssetMgr = NewObject<UWebBrowserAssetManager>((UObject*)GetTransientPackage(), NAME_None, RF_Transient | RF_Public);
@@ -27,14 +28,16 @@ public:
 				WebBrowserSingleton->SetDefaultTranslucentMaterial(WebBrowserAssetMgr->GetDefaultTranslucentMaterial());
 			}
 		}
+#endif
 	}
 
 	virtual void ShutdownModule() override
 	{
 	}
-	
 private:
+#if WITH_EDITOR || PLATFORM_ANDROID || PLATFORM_IOS
 	UWebBrowserAssetManager* WebBrowserAssetMgr;
+#endif
 };
 
 //////////////////////////////////////////////////////////////////////////
diff --git a/Engine/Plugins/Runtime/WebBrowserWidget/Source/WebBrowserWidget/Public/WebBrowserAssetManager.h b/Engine/Plugins/Runtime/WebBrowserWidget/Source/WebBrowserWidget/Public/WebBrowserAssetManager.h
--- a/Engine/Plugins/Runtime/WebBrowserWidget/Source/WebBrowserWidget/Public/WebBrowserAssetManager.h
+++ b/Engine/Plugins/Runtime/WebBrowserWidget/Source/WebBrowserWidget/Public/WebBrowserAssetManager.h
@@ -1,6 +1,7 @@
 ﻿// Copyright 1998-2019 Epic Games, Inc. All Rights Reserved.
 
 #pragma once
+#if WITH_EDITOR || PLATFORM_ANDROID || PLATFORM_IOS
 
 #include "UObject/SoftObjectPtr.h"
 #include "Materials/Material.h"
@@ -26,3 +27,5 @@ protected:
 	TSoftObjectPtr<UMaterial> DefaultMaterial;
 	TSoftObjectPtr<UMaterial> DefaultTranslucentMaterial;
 };
+
+#endif
\ No newline at end of file
