Fix for Mac non-unity build

diff --git a/Engine/Plugins/Runtime/WebBrowserWidget/Source/WebBrowserWidget/Private/WebBrowserAssetManager.cpp b/Engine/Plugins/Runtime/WebBrowserWidget/Source/WebBrowserWidget/Private/WebBrowserAssetManager.cpp
--- a/Engine/Plugins/Runtime/WebBrowserWidget/Source/WebBrowserWidget/Private/WebBrowserAssetManager.cpp
+++ b/Engine/Plugins/Runtime/WebBrowserWidget/Source/WebBrowserWidget/Private/WebBrowserAssetManager.cpp
@@ -3,7 +3,6 @@
 #include "WebBrowserAssetManager.h"
 
 #if WITH_EDITOR || PLATFORM_ANDROID || PLATFORM_IOS
-#include "Materials/Material.h"
 #include "WebBrowserTexture.h"
 #endif
 
diff --git a/Engine/Plugins/Runtime/WebBrowserWidget/Source/WebBrowserWidget/Public/WebBrowserAssetManager.h b/Engine/Plugins/Runtime/WebBrowserWidget/Source/WebBrowserWidget/Public/WebBrowserAssetManager.h
--- a/Engine/Plugins/Runtime/WebBrowserWidget/Source/WebBrowserWidget/Public/WebBrowserAssetManager.h
+++ b/Engine/Plugins/Runtime/WebBrowserWidget/Source/WebBrowserWidget/Public/WebBrowserAssetManager.h
@@ -3,7 +3,7 @@
 #pragma once
 
 #include "UObject/SoftObjectPtr.h"
-
+#include "Materials/Material.h"
 #include "WebBrowserAssetManager.generated.h"
 
 class UMaterial;
