- Don't crash if the CEF library fails to load. We see reports in the field from users with corrupt disk images missing files or running OS configurations that don't support the CEF library, in those cases just fail to load CEF but don't crash.

diff --git a/Engine/Plugins/Runtime/WebBrowserWidget/Source/WebBrowserWidget/Private/WebBrowserWidgetModule.cpp b/Engine/Plugins/Runtime/WebBrowserWidget/Source/WebBrowserWidget/Private/WebBrowserWidgetModule.cpp
--- a/Engine/Plugins/Runtime/WebBrowserWidget/Source/WebBrowserWidget/Private/WebBrowserWidgetModule.cpp
+++ b/Engine/Plugins/Runtime/WebBrowserWidget/Source/WebBrowserWidget/Private/WebBrowserWidgetModule.cpp
@@ -20,11 +20,15 @@ public:
 			WebBrowserAssetMgr = NewObject<UWebBrowserAssetManager>((UObject*)GetTransientPackage(), NAME_None, RF_Transient | RF_Public);
 			WebBrowserAssetMgr->LoadDefaultMaterials();
 
-			IWebBrowserSingleton* WebBrowserSingleton = IWebBrowserModule::Get().GetSingleton();
-			if (WebBrowserSingleton)
+			IWebBrowserModule::Get(); // force the module to load
+			if (IWebBrowserModule::IsAvailable() && IWebBrowserModule::Get().IsWebModuleAvailable())
 			{
-				WebBrowserSingleton->SetDefaultMaterial(WebBrowserAssetMgr->GetDefaultMaterial());
-				WebBrowserSingleton->SetDefaultTranslucentMaterial(WebBrowserAssetMgr->GetDefaultTranslucentMaterial());
+				IWebBrowserSingleton* WebBrowserSingleton = IWebBrowserModule::Get().GetSingleton();
+				if (WebBrowserSingleton)
+				{
+					WebBrowserSingleton->SetDefaultMaterial(WebBrowserAssetMgr->GetDefaultMaterial());
+					WebBrowserSingleton->SetDefaultTranslucentMaterial(WebBrowserAssetMgr->GetDefaultTranslucentMaterial());
+				}
 			}
 		}
 	}
diff --git a/Engine/Source/Runtime/WebBrowser/Private/SWebBrowserView.cpp b/Engine/Source/Runtime/WebBrowser/Private/SWebBrowserView.cpp
--- a/Engine/Source/Runtime/WebBrowser/Private/SWebBrowserView.cpp
+++ b/Engine/Source/Runtime/WebBrowser/Private/SWebBrowserView.cpp
@@ -116,7 +116,11 @@ void SWebBrowserView::Construct(const FArguments& InArgs, const TSharedPtr<IWebB
 			Settings.Context = InArgs._ContextSettings;
 			Settings.AltRetryDomains = InArgs._AltRetryDomains;
 
-			BrowserWindow = IWebBrowserModule::Get().GetSingleton()->CreateBrowserWindow(Settings);
+			// IWebBrowserModule::Get() was already callled in WebBrowserWidgetModule.cpp so we don't need to force the load again here
+			if (IWebBrowserModule::IsAvailable() && IWebBrowserModule::Get().IsWebModuleAvailable())
+			{
+				BrowserWindow = IWebBrowserModule::Get().GetSingleton()->CreateBrowserWindow(Settings);
+			}
 		}
 	}
 
