Allow WebGL and video playback in 3D WebBrowser on Android and fix issue with Oculus hardware acceleration

#android

diff --git a/Engine/Build/Android/Java/src/com/epicgames/unreal/WebViewControl.java b/Engine/Build/Android/Java/src/com/epicgames/unreal/WebViewControl.java
--- a/Engine/Build/Android/Java/src/com/epicgames/unreal/WebViewControl.java
+++ b/Engine/Build/Android/Java/src/com/epicgames/unreal/WebViewControl.java
@@ -1880,7 +1880,7 @@ class WebViewControl
 			{
 				try 
 				{
-					glAttachedCanvas = mSurface.lockCanvas(null);
+					glAttachedCanvas = mSurface.lockHardwareCanvas();
 					//draw the view to provided canvas
 
 					//translate canvas to reflect view scrolling
diff --git a/Engine/Plugins/Runtime/WebBrowserWidget/Source/WebBrowserWidget/WebBrowserWidget.build.cs b/Engine/Plugins/Runtime/WebBrowserWidget/Source/WebBrowserWidget/WebBrowserWidget.build.cs
--- a/Engine/Plugins/Runtime/WebBrowserWidget/Source/WebBrowserWidget/WebBrowserWidget.build.cs
+++ b/Engine/Plugins/Runtime/WebBrowserWidget/Source/WebBrowserWidget/WebBrowserWidget.build.cs
@@ -1,4 +1,7 @@
-﻿// Copyright Epic Games, Inc. All Rights Reserved.
+// Copyright Epic Games, Inc. All Rights Reserved.
+
+using UnrealBuildTool;
+using System.IO;
 
 namespace UnrealBuildTool.Rules
 {
@@ -53,6 +56,12 @@ namespace UnrealBuildTool.Rules
                         "WebBrowserTexture",
                     }
                 );
+
+				if (Target.Platform == UnrealTargetPlatform.Android)
+				{
+					string PluginPath = Utils.MakePathRelativeTo(ModuleDirectory, Target.RelativeEnginePath);
+					AdditionalPropertiesForReceipt.Add("AndroidPlugin", Path.Combine(PluginPath, "WebBrowserWidget_UPL.xml"));
+				}
 			}
         }
     }
diff --git a/Engine/Plugins/Runtime/WebBrowserWidget/Source/WebBrowserWidget/WebBrowserWidget_UPL.xml b/Engine/Plugins/Runtime/WebBrowserWidget/Source/WebBrowserWidget/WebBrowserWidget_UPL.xml
new file mode 100644
--- /dev/null
+++ b/Engine/Plugins/Runtime/WebBrowserWidget/Source/WebBrowserWidget/WebBrowserWidget_UPL.xml
@@ -0,0 +1,31 @@
+<?xml version="1.0" encoding="utf-8"?>
+<root xmlns:android="http://schemas.android.com/apk/res/android">
+	<init>
+		<log text="WebBrowserWidget Plugin Init"/>
+		<setBoolFromPropertyContains result="bPackageForOculusQuest" ini="Engine" section="/Script/AndroidRuntimeSettings.AndroidRuntimeSettings" property="PackageForOculusMobile" contains="Quest"/>
+		<setBoolFromPropertyContains result="bPackageForOculusQuest2" ini="Engine" section="/Script/AndroidRuntimeSettings.AndroidRuntimeSettings" property="PackageForOculusMobile" contains="Quest2"/>
+		<setBoolFromPropertyContains result="bPackageForOculusQuestPro" ini="Engine" section="/Script/AndroidRuntimeSettings.AndroidRuntimeSettings" property="PackageForOculusMobile" contains="QuestPro"/>
+
+		<setBoolOr result="bPackageForOculusMobile" 
+			arg1="$B(bPackageForOculusQuest)" arg2="$B(bPackageForOculusQuest2)"/>
+		<setBoolOr result="bPackageForOculusMobile" 
+			arg1="$B(bPackageForOculusMobile)" arg2="$B(bPackageForOculusQuestPro)"/>
+
+		<if condition="bPackageForOculusMobile">
+			<true>
+				<log text="Adding com.oculus.always_draw_view_root for Oculus devices to manifest for hardware acceleration."/>
+			</true>
+		</if>
+	</init>
+
+	<androidManifestUpdates>
+		<if condition="bPackageForOculusMobile">
+			<true>
+				<addElements tag="application">
+					<meta-data android:name="com.oculus.always_draw_view_root" android:value="true"/>
+				</addElements>
+			</true>
+		</if>
+	</androidManifestUpdates>
+
+</root>
diff --git a/Engine/Source/Runtime/WebBrowser/Private/Android/AndroidWebBrowserWidget.cpp b/Engine/Source/Runtime/WebBrowser/Private/Android/AndroidWebBrowserWidget.cpp
--- a/Engine/Source/Runtime/WebBrowser/Private/Android/AndroidWebBrowserWidget.cpp
+++ b/Engine/Source/Runtime/WebBrowser/Private/Android/AndroidWebBrowserWidget.cpp
@@ -114,7 +114,7 @@ void SAndroidWebBrowserWidget::Construct(const FArguments& Args)
 	WebBrowserWindowPtr = Args._WebBrowserWindow;
 	IsAndroid3DBrowser = true;
 
-	bShouldUseBitmapRender = AndroidThunkCpp_IsOculusMobileApplication();
+	bShouldUseBitmapRender = false; //AndroidThunkCpp_IsOculusMobileApplication();
 	bMouseCapture = false;
 
 	HistorySize = 0;
