Fixed some implicit FSoftObjectPath construction.

diff --git a/Engine/Plugins/Experimental/Avalanche/Source/AvalancheEditorCore/Private/AvaEditor.cpp b/Engine/Plugins/Experimental/Avalanche/Source/AvalancheEditorCore/Private/AvaEditor.cpp
--- a/Engine/Plugins/Experimental/Avalanche/Source/AvalancheEditorCore/Private/AvaEditor.cpp
+++ b/Engine/Plugins/Experimental/Avalanche/Source/AvalancheEditorCore/Private/AvaEditor.cpp
@@ -880,7 +880,7 @@ bool FAvaEditor::PasteFromString(FString& InPastedData) const
 			// Pop the Temporary Tag added on Copy, which would be the Actor Path of Source Actor
 			const FName CopiedTag = InPastedActor->Tags.Pop();
 
-			TSoftObjectPtr<AActor> SourceActor(CopiedTag.ToString());
+			TSoftObjectPtr<AActor> SourceActor = TSoftObjectPtr<AActor>(FSoftObjectPath(CopiedTag.ToString()));
 
 			PastedActors.Emplace(InPastedActor, MoveTemp(SourceActor));
 		}
diff --git a/Engine/Plugins/Runtime/WebBrowserWidget/Source/WebBrowserWidget/Private/WebBrowserAssetManager.cpp b/Engine/Plugins/Runtime/WebBrowserWidget/Source/WebBrowserWidget/Private/WebBrowserAssetManager.cpp
--- a/Engine/Plugins/Runtime/WebBrowserWidget/Source/WebBrowserWidget/Private/WebBrowserAssetManager.cpp
+++ b/Engine/Plugins/Runtime/WebBrowserWidget/Source/WebBrowserWidget/Private/WebBrowserAssetManager.cpp
@@ -12,8 +12,8 @@
 
 UWebBrowserAssetManager::UWebBrowserAssetManager(const FObjectInitializer& ObjectInitializer)
 	: Super(ObjectInitializer) ,
-	DefaultMaterial(FString(TEXT("/WebBrowserWidget/WebTexture_M.WebTexture_M"))) ,
-	DefaultTranslucentMaterial(FString(TEXT("/WebBrowserWidget/WebTexture_TM.WebTexture_TM")))
+	DefaultMaterial(FSoftObjectPath(FString(TEXT("/WebBrowserWidget/WebTexture_M.WebTexture_M")))),
+	DefaultTranslucentMaterial(FSoftObjectPath(FString(TEXT("/WebBrowserWidget/WebTexture_TM.WebTexture_TM"))))
 {
 #if WITH_EDITOR || PLATFORM_ANDROID || PLATFORM_IOS
 	// Add a hard reference to UWebBrowserTexture, without this the WebBrowserTexture DLL never gets loaded on Windows.
diff --git a/Engine/Plugins/Runtime/nDisplay/Source/DisplayClusterLightCardEditor/Private/DisplayClusterLightCardEditorUtils.cpp b/Engine/Plugins/Runtime/nDisplay/Source/DisplayClusterLightCardEditor/Private/DisplayClusterLightCardEditorUtils.cpp
--- a/Engine/Plugins/Runtime/nDisplay/Source/DisplayClusterLightCardEditor/Private/DisplayClusterLightCardEditorUtils.cpp
+++ b/Engine/Plugins/Runtime/nDisplay/Source/DisplayClusterLightCardEditor/Private/DisplayClusterLightCardEditorUtils.cpp
@@ -62,7 +62,7 @@ TSet<UClass*> UE::DisplayClusterLightCardEditorUtils::GetAllStageActorClasses()
 			const FTopLevelAssetPath ClassObjectPath(FPackageName::ExportTextPathToObjectPath(*GeneratedClassPathPtr));
 			if (InheritedInterfacePaths.Contains(ClassObjectPath))
 			{
-				TSoftObjectPtr<UClass> SoftObjectPtr(ClassObjectPath.ToString());
+				TSoftObjectPtr<UClass> SoftObjectPtr = TSoftObjectPtr<UClass>(FSoftObjectPath(ClassObjectPath.ToString()));
 				if (UClass* Class = SoftObjectPtr.LoadSynchronous())
 				{
 					OutClasses.Add(Class);
