moving 3D browser assets into engine to allow ApplePlatformWebBrowser to use indirect rendering path even when the UMG Web Browser Widget plugin is disabled

#Synced-CL Unknown

diff --git a/Engine/Plugins/Runtime/WebBrowserWidget/Source/WebBrowserWidget/Private/WebBrowserAssetManager.cpp b/Engine/Plugins/Runtime/WebBrowserWidget/Source/WebBrowserWidget/Private/WebBrowserAssetManager.cpp
--- a/Engine/Plugins/Runtime/WebBrowserWidget/Source/WebBrowserWidget/Private/WebBrowserAssetManager.cpp
+++ b/Engine/Plugins/Runtime/WebBrowserWidget/Source/WebBrowserWidget/Private/WebBrowserAssetManager.cpp
@@ -12,8 +12,8 @@
 
 UWebBrowserAssetManager::UWebBrowserAssetManager(const FObjectInitializer& ObjectInitializer)
 	: Super(ObjectInitializer) ,
-	DefaultMaterial(FSoftObjectPath(FString(TEXT("/WebBrowserWidget/WebTexture_M.WebTexture_M")))),
-	DefaultTranslucentMaterial(FSoftObjectPath(FString(TEXT("/WebBrowserWidget/WebTexture_TM.WebTexture_TM"))))
+	DefaultMaterial(FSoftObjectPath(FString(TEXT("/Engine/WebBrowser/WebTexture_M.WebTexture_M")))),
+	DefaultTranslucentMaterial(FSoftObjectPath(FString(TEXT("/Engine/WebBrowser/WebTexture_TM.WebTexture_TM"))))
 {
 #if WITH_EDITOR || PLATFORM_ANDROID || PLATFORM_IOS
 	// Add a hard reference to UWebBrowserTexture, without this the WebBrowserTexture DLL never gets loaded on Windows.
diff --git a/Engine/Plugins/Runtime/WebBrowserWidget/Source/WebBrowserWidget/Private/WebBrowserWidgetModule.cpp b/Engine/Plugins/Runtime/WebBrowserWidget/Source/WebBrowserWidget/Private/WebBrowserWidgetModule.cpp
--- a/Engine/Plugins/Runtime/WebBrowserWidget/Source/WebBrowserWidget/Private/WebBrowserWidgetModule.cpp
+++ b/Engine/Plugins/Runtime/WebBrowserWidget/Source/WebBrowserWidget/Private/WebBrowserWidgetModule.cpp
@@ -14,30 +14,12 @@ class FWebBrowserWidgetModule : public IWebBrowserWidgetModule
 public:
 	virtual void StartupModule() override
 	{
-		if (WebBrowserAssetMgr == nullptr)
-		{
-			WebBrowserAssetMgr = NewObject<UWebBrowserAssetManager>((UObject*)GetTransientPackage(), NAME_None, RF_Transient | RF_Public);
-			WebBrowserAssetMgr->LoadDefaultMaterials();
-
-			IWebBrowserModule::Get(); // force the module to load
-			if (IWebBrowserModule::IsAvailable() && IWebBrowserModule::Get().IsWebModuleAvailable())
-			{
-				IWebBrowserSingleton* WebBrowserSingleton = IWebBrowserModule::Get().GetSingleton();
-				if (WebBrowserSingleton)
-				{
-					WebBrowserSingleton->SetDefaultMaterial(WebBrowserAssetMgr->GetDefaultMaterial());
-					WebBrowserSingleton->SetDefaultTranslucentMaterial(WebBrowserAssetMgr->GetDefaultTranslucentMaterial());
-				}
-			}
-		}
+		IWebBrowserModule::Get(); // force the module to load
 	}
 
 	virtual void ShutdownModule() override
 	{
 	}
-	
-private:
-	UWebBrowserAssetManager* WebBrowserAssetMgr;
 };
 
 //////////////////////////////////////////////////////////////////////////
diff --git a/Engine/Plugins/Runtime/WebBrowserWidget/Source/WebBrowserWidget/Public/WebBrowserAssetManager.h b/Engine/Plugins/Runtime/WebBrowserWidget/Source/WebBrowserWidget/Public/WebBrowserAssetManager.h
--- a/Engine/Plugins/Runtime/WebBrowserWidget/Source/WebBrowserWidget/Public/WebBrowserAssetManager.h
+++ b/Engine/Plugins/Runtime/WebBrowserWidget/Source/WebBrowserWidget/Public/WebBrowserAssetManager.h
@@ -10,15 +10,19 @@ class UMaterial;
 /**
  * 
  */
+
 UCLASS()
 class WEBBROWSERWIDGET_API UWebBrowserAssetManager : public UObject
 {
 	GENERATED_UCLASS_BODY()
 
 public:
+	UE_DEPRECATED(5.8, "UWebBrowserAssetManager is deprecated. LoadDefaultMaterials is no longer necessary for the GetDefaultMaterial functions in IWebBrowserSingleton.")
 	void LoadDefaultMaterials();
 
+	UE_DEPRECATED(5.8, "UWebBrowserAssetManager is deprecated. Please use the GetDefaultMaterial functions in IWebBrowserSingleton instead.")
 	UMaterial* GetDefaultMaterial(); 
+	UE_DEPRECATED(5.8, "UWebBrowserAssetManager is deprecated. Please use the GetDefaultMaterial functions in IWebBrowserSingleton instead.")
 	UMaterial* GetDefaultTranslucentMaterial(); 
 
 protected:
diff --git a/Engine/Source/Runtime/WebBrowser/Private/WebBrowserSingleton.cpp b/Engine/Source/Runtime/WebBrowser/Private/WebBrowserSingleton.cpp
--- a/Engine/Source/Runtime/WebBrowser/Private/WebBrowserSingleton.cpp
+++ b/Engine/Source/Runtime/WebBrowser/Private/WebBrowserSingleton.cpp
@@ -247,8 +247,8 @@ FWebBrowserSingleton::FWebBrowserSingleton(const FWebBrowserInitSettings& WebBro
 #if WITH_CEF3
 	, bCEFInitialized(false)
 #endif
-	, DefaultMaterial(nullptr)
-	, DefaultTranslucentMaterial(nullptr)
+	, DefaultMaterial(FSoftObjectPath(FString(TEXT("/Engine/WebBrowser/WebTexture_M.WebTexture_M"))))
+	, DefaultTranslucentMaterial(FSoftObjectPath(FString(TEXT("/Engine/WebBrowser/WebTexture_TM.WebTexture_TM"))))
 {
 #if WITH_CEF3
 	
@@ -391,7 +391,9 @@ FWebBrowserSingleton::FWebBrowserSingleton(const FWebBrowserInitSettings& WebBro
 #elif PLATFORM_ANDROID
 	DefaultCookieManager = MakeShareable(new FAndroidCookieManager());
 #endif
-
+	
+	DefaultMaterial.LoadSynchronous();
+	DefaultTranslucentMaterial.LoadSynchronous();
 }
 
 
diff --git a/Engine/Source/Runtime/WebBrowser/Private/WebBrowserSingleton.h b/Engine/Source/Runtime/WebBrowser/Private/WebBrowserSingleton.h
--- a/Engine/Source/Runtime/WebBrowser/Private/WebBrowserSingleton.h
+++ b/Engine/Source/Runtime/WebBrowser/Private/WebBrowserSingleton.h
@@ -6,6 +6,8 @@
 #include "Containers/Ticker.h"
 #include "IWebBrowserSingleton.h"
 
+#include "Materials/Material.h"
+
 #if WITH_CEF3
 #if PLATFORM_WINDOWS
 	#include "Windows/AllowWindowsPlatformTypes.h"
@@ -133,13 +135,13 @@ public:
 	/** Get a reference to UWebBrowser's default material*/
 	virtual UMaterialInterface* GetDefaultMaterial() override
 	{
-		return DefaultMaterial;
+		return DefaultMaterial.Get();
 	}
 
 	/** Get a reference to UWebBrowser's translucent material*/
 	virtual UMaterialInterface* GetDefaultTranslucentMaterial() override
 	{
-		return DefaultTranslucentMaterial;
+		return DefaultTranslucentMaterial.Get();
 	}
 
 public:
@@ -206,10 +208,10 @@ private:
 #endif
 
 	/** Reference to UWebBrowser's default material*/
-	UMaterialInterface* DefaultMaterial;
+	TSoftObjectPtr<UMaterialInterface> DefaultMaterial;
 
 	/** Reference to UWebBrowser's translucent material*/
-	UMaterialInterface* DefaultTranslucentMaterial;
+	TSoftObjectPtr<UMaterialInterface> DefaultTranslucentMaterial;
 
 	/* Time when our Tick() was last called. */
 	double PreviousTickTimeSeconds = 0.0;
